PR #9002: Fix crash with ReplayTracks plugin (Contributed by [REDACTED_PERSON])

diff --git a/Engine/Plugins/MovieScene/ReplayTracks/Source/ReplayTracksEditor/Private/TrackEditors/ReplayTrackEditor.cpp b/Engine/Plugins/MovieScene/ReplayTracks/Source/ReplayTracksEditor/Private/TrackEditors/ReplayTrackEditor.cpp
--- a/Engine/Plugins/MovieScene/ReplayTracks/Source/ReplayTracksEditor/Private/TrackEditors/ReplayTrackEditor.cpp
+++ b/Engine/Plugins/MovieScene/ReplayTracks/Source/ReplayTracksEditor/Private/TrackEditors/ReplayTrackEditor.cpp
@@ -252,7 +252,10 @@ void FReplayTrackEditor::OnInitialize()
 void FReplayTrackEditor::OnRelease()
 {
 	const TSharedPtr<ISequencer> SequencerPtr = GetSequencer();
-	SequencerPtr->OnGlobalTimeChanged().Remove(GlobalTimeChangedHandle);
+	if (SequencerPtr.IsValid())
+	{
+		SequencerPtr->OnGlobalTimeChanged().Remove(GlobalTimeChangedHandle);
+	}
 }
 
 void FReplayTrackEditor::Tick(float DeltaTime)
