Sequencer: Check for valid spectator pawn before dereferencing

diff --git a/Engine/Plugins/MovieScene/ReplayTracks/Source/ReplayTracksEditor/Private/TrackEditors/ReplayTrackEditor.cpp b/Engine/Plugins/MovieScene/ReplayTracks/Source/ReplayTracksEditor/Private/TrackEditors/ReplayTrackEditor.cpp
--- a/Engine/Plugins/MovieScene/ReplayTracks/Source/ReplayTracksEditor/Private/TrackEditors/ReplayTrackEditor.cpp
+++ b/Engine/Plugins/MovieScene/ReplayTracks/Source/ReplayTracksEditor/Private/TrackEditors/ReplayTrackEditor.cpp
@@ -394,8 +394,10 @@ void FReplayTrackEditor::MovePIEViewTargetToLockedActor(UWorld* PlaybackWorld)
 
 			// Rotate the pawn... most of the time that's going to rotate the same thing we just rotated above (ViewTargetSceneComponent)
 			// but it's not necessarily the case.
-			ASpectatorPawn* SpectatorPawn = PlayerController->GetSpectatorPawn();
-			SpectatorPawn->SetActorRotation(Rotation);
+			if (ASpectatorPawn* SpectatorPawn = PlayerController->GetSpectatorPawn())
+			{
+				SpectatorPawn->SetActorRotation(Rotation);
+			}
 
 			// Rotate the player controller, because the view target doesn't always solely define the final camera view. Most of the time,
 			// the player controller's own rotation is inserted into the mix.
