Fixed includes after merge

diff --git a/Engine/Plugins/Experimental/ChaosFlesh/Source/ChaosFlesh/Private/ChaosFlesh/GEO.cpp b/Engine/Plugins/Experimental/ChaosFlesh/Source/ChaosFlesh/Private/ChaosFlesh/GEO.cpp
--- a/Engine/Plugins/Experimental/ChaosFlesh/Source/ChaosFlesh/Private/ChaosFlesh/GEO.cpp
+++ b/Engine/Plugins/Experimental/ChaosFlesh/Source/ChaosFlesh/Private/ChaosFlesh/GEO.cpp
@@ -2,6 +2,7 @@
 #include "ChaosFlesh/GEO.h"
 #include "ChaosFlesh/IFileStream.h"
 #include "HAL/PlatformFile.h"
+#include "HAL/PlatformFileManager.h"
 
 #ifdef USE_ZLIB
 #include "ChaosFlesh/ZIP.h"
diff --git a/Engine/Plugins/MovieScene/ReplayTracks/Source/ReplayTracks/Private/Systems/MovieSceneReplaySystem.cpp b/Engine/Plugins/MovieScene/ReplayTracks/Source/ReplayTracks/Private/Systems/MovieSceneReplaySystem.cpp
--- a/Engine/Plugins/MovieScene/ReplayTracks/Source/ReplayTracks/Private/Systems/MovieSceneReplaySystem.cpp
+++ b/Engine/Plugins/MovieScene/ReplayTracks/Source/ReplayTracks/Private/Systems/MovieSceneReplaySystem.cpp
@@ -7,7 +7,6 @@
 #include "EntitySystem/BuiltInComponentTypes.h"
 #include "EntitySystem/MovieSceneEntitySystemTask.h"
 #include "EntitySystem/MovieSceneInstanceRegistry.h"
-#include "EntitySystem/MovieSceneMasterInstantiatorSystem.h"
 #include "EntitySystem/MovieSceneRootInstantiatorSystem.h"
 #include "Evaluation/MovieSceneEvaluationTemplateInstance.h"
 #include "GameDelegates.h"
