Sequencer: remove uses of IMovieScenePlayer

diff --git a/Engine/Plugins/MovieScene/ReplayTracks/Source/ReplayTracks/Private/Systems/MovieSceneReplaySystem.cpp b/Engine/Plugins/MovieScene/ReplayTracks/Source/ReplayTracks/Private/Systems/MovieSceneReplaySystem.cpp
--- a/Engine/Plugins/MovieScene/ReplayTracks/Source/ReplayTracks/Private/Systems/MovieSceneReplaySystem.cpp
+++ b/Engine/Plugins/MovieScene/ReplayTracks/Source/ReplayTracks/Private/Systems/MovieSceneReplaySystem.cpp
@@ -226,10 +226,10 @@ void UMovieSceneReplaySystem::OnRunEvaluation()
 	const FInstanceRegistry* InstanceRegistry = Linker->GetInstanceRegistry();
 	const FSequenceInstance& Instance = InstanceRegistry->GetInstance(ActiveReplayInfo.InstanceHandle);
 
-	IMovieScenePlayer* Player = Instance.GetPlayer();
+	TSharedRef<const FSharedPlaybackState> SharedPlaybackState = Instance.GetSharedPlaybackState();
 	const FMovieSceneContext& Context = Instance.GetContext();
 
-	UWorld* World = Player->GetPlaybackContext()->GetWorld();
+	UWorld* World = SharedPlaybackState->GetPlaybackContext()->GetWorld();
 	UDemoNetDriver* DemoNetDriver = World->GetDemoNetDriver();
 
 	if (DemoNetDriver)
@@ -247,7 +247,7 @@ void UMovieSceneReplaySystem::OnRunEvaluation()
 		}
 
 		// Set time dilation and current demo time according to our current sequencer playback.
-		const EMovieScenePlayerStatus::Type CurrentPlayerStatus = Player->GetPlaybackStatus();
+		const EMovieScenePlayerStatus::Type CurrentPlayerStatus = Context.GetStatus();
 
 		const bool bIsPlaying = (CurrentPlayerStatus == EMovieScenePlayerStatus::Playing);
 		const bool bWasPlaying = (PreviousPlayerStatus == EMovieScenePlayerStatus::Playing);
@@ -330,9 +330,9 @@ void UMovieSceneReplaySystem::StartReplay(const FReplayInfo& ReplayInfo)
 
 	const FInstanceRegistry* InstanceRegistry = Linker->GetInstanceRegistry();
 	const FSequenceInstance& Instance = InstanceRegistry->GetInstance(ReplayInfo.InstanceHandle);
-	IMovieScenePlayer* Player = Instance.GetPlayer();
+	TSharedRef<const FSharedPlaybackState> SharedPlaybackState = Instance.GetSharedPlaybackState();
 	const FMovieSceneContext Context = Instance.GetContext();
-	UWorld* World = Player->GetPlaybackContext()->GetWorld();
+	UWorld* World = SharedPlaybackState->GetPlaybackContext()->GetWorld();
 
 	UDemoNetDriver* DemoNetDriver = World->GetDemoNetDriver();
 	if (!ensure(DemoNetDriver == nullptr))
@@ -359,11 +359,11 @@ void UMovieSceneReplaySystem::StartReplay(const FReplayInfo& ReplayInfo)
 	// We have a few things to do just before/after the map has been loaded.
 	if (!PreLoadMapHandle.IsValid())
 	{
-		PreLoadMapHandle = FCoreUObjectDelegates::PreLoadMap.AddStatic(UMovieSceneReplaySystem::OnPreLoadMap, Player);
+		PreLoadMapHandle = FCoreUObjectDelegates::PreLoadMap.AddStatic(UMovieSceneReplaySystem::OnPreLoadMap, SharedPlaybackState);
 	}
 	if (!PostLoadMapHandle.IsValid())
 	{
-		PostLoadMapHandle = FCoreUObjectDelegates::PostLoadMapWithWorld.AddStatic(UMovieSceneReplaySystem::OnPostLoadMap, Player, Context);
+		PostLoadMapHandle = FCoreUObjectDelegates::PostLoadMapWithWorld.AddStatic(UMovieSceneReplaySystem::OnPostLoadMap, SharedPlaybackState, Context);
 	}
 	if (!EndPlayMapHandle.IsValid())
 	{
@@ -389,8 +389,8 @@ void UMovieSceneReplaySystem::StopReplay(const FReplayInfo& ReplayInfo)
 
 	const FInstanceRegistry* InstanceRegistry = Linker->GetInstanceRegistry();
 	const FSequenceInstance& Instance = InstanceRegistry->GetInstance(ReplayInfo.InstanceHandle);
-	IMovieScenePlayer* Player = Instance.GetPlayer();
-	UWorld* World = Player->GetPlaybackContext()->GetWorld();
+	TSharedRef<const FSharedPlaybackState> SharedPlaybackState = Instance.GetSharedPlaybackState();
+	UWorld* World = SharedPlaybackState->GetPlaybackContext()->GetWorld();
 	UGameInstance* GameInstance = World->GetGameInstance();
 	UReplaySubsystem* ReplaySubsystem = GameInstance ? GameInstance->GetSubsystem<UReplaySubsystem>() : nullptr;
 	if (ReplaySubsystem != nullptr)
@@ -410,7 +410,7 @@ void UMovieSceneReplaySystem::StopReplay(const FReplayInfo& ReplayInfo)
 	bReplayActive = false;
 }
 
-void UMovieSceneReplaySystem::OnPreLoadMap(const FString& MapName, IMovieScenePlayer* Player)
+void UMovieSceneReplaySystem::OnPreLoadMap(const FString& MapName, TSharedRef<const UE::MovieScene::FSharedPlaybackState> SharedPlaybackState)
 {
 	FCoreUObjectDelegates::PreLoadMap.Remove(PreLoadMapHandle);
 	PreLoadMapHandle.Reset();
@@ -421,19 +421,21 @@ void UMovieSceneReplaySystem::OnPreLoadMap(const FString& MapName, IMovieScenePl
 	// We could call "Finish" here but we dont' want that either because it actually re-evaluates the sequence one last time, 
 	// which would re-trigger the replay and re-re-load the replay map again.
 	// So we just do the minimum we need here.
-	FMovieSceneSpawnRegister& SpawnRegister = Player->GetSpawnRegister();
-	SpawnRegister.ForgetExternallyOwnedSpawnedObjects(Player->State, *Player);
-	SpawnRegister.CleanUp(*Player);
+	if (FMovieSceneSpawnRegister* SpawnRegister = SharedPlaybackState->FindCapability<FMovieSceneSpawnRegister>())
+	{
+		SpawnRegister->ForgetExternallyOwnedSpawnedObjects(SharedPlaybackState);
+		SpawnRegister->CleanUp(SharedPlaybackState);
+	}
 }
 
-void UMovieSceneReplaySystem::OnPostLoadMap(UWorld* World, IMovieScenePlayer* LastPlayer, FMovieSceneContext LastContext)
+void UMovieSceneReplaySystem::OnPostLoadMap(UWorld* World, TSharedRef<const UE::MovieScene::FSharedPlaybackState> SharedPlaybackState, FMovieSceneContext LastContext)
 {
 	FCoreUObjectDelegates::PostLoadMapWithWorld.Remove(PostLoadMapHandle);
 	PostLoadMapHandle.Reset();
 
 	// After the map has loaded, we wait for the game to be ready to start showing the replay. This generally includes waiting
 	// for the replay pawn and player controller.
-	World->GetTimerManager().SetTimer(ReEvaluateHandle, [World, LastPlayer, LastContext]()
+	World->GetTimerManager().SetTimer(ReEvaluateHandle, [World, SharedPlaybackState, LastContext]()
 		{
 			FMovieSceneReplayManager& Manager = FMovieSceneReplayManager::Get();
 			FMovieSceneReplayBroker* Broker = Manager.FindBroker(World);
@@ -441,8 +443,12 @@ void UMovieSceneReplaySystem::OnPostLoadMap(UWorld* World, IMovieScenePlayer* La
 			{
 				World->GetTimerManager().ClearTimer(ReEvaluateHandle);
 
-				FMovieSceneRootEvaluationTemplateInstance& RootEvalTemplate = LastPlayer->GetEvaluationTemplate();
-				RootEvalTemplate.EvaluateSynchronousBlocking(LastContext);
+				// Re-evaluate synchronously.
+				if (TSharedPtr<FMovieSceneEntitySystemRunner> Runner = SharedPlaybackState->GetRunner())
+				{
+					Runner->QueueUpdate(LastContext, SharedPlaybackState->GetRootInstanceHandle());
+					Runner->Flush();
+				}
 			}
 		},
 		0.1f,
diff --git a/Engine/Plugins/MovieScene/ReplayTracks/Source/ReplayTracks/Private/Systems/MovieSceneReplaySystem.h b/Engine/Plugins/MovieScene/ReplayTracks/Source/ReplayTracks/Private/Systems/MovieSceneReplaySystem.h
--- a/Engine/Plugins/MovieScene/ReplayTracks/Source/ReplayTracks/Private/Systems/MovieSceneReplaySystem.h
+++ b/Engine/Plugins/MovieScene/ReplayTracks/Source/ReplayTracks/Private/Systems/MovieSceneReplaySystem.h
@@ -9,7 +9,6 @@
 
 class FDelegateHandle;
 class IConsoleVariable;
-class IMovieScenePlayer;
 struct FMovieSceneContext;
 struct FTimerHandle;
 
@@ -20,6 +19,8 @@ namespace UE
 namespace MovieScene
 {
 
+struct FSharedPlaybackState;
+
 struct FReplayComponentData
 {
 	FReplayComponentData() : Section(nullptr) {}
@@ -93,8 +94,8 @@ private:
 	// Handlers for replay events. These are all static because they happen across level reloads,
 	// which mean that the current object may have been wiped out and re-created, so we can't rely
 	// on any instance data.
-	static void OnPreLoadMap(const FString& MapName, IMovieScenePlayer* Player);
-	static void OnPostLoadMap(UWorld* World, IMovieScenePlayer* LastPlayer, FMovieSceneContext LastContext);
+	static void OnPreLoadMap(const FString& MapName, TSharedRef<const UE::MovieScene::FSharedPlaybackState> SharedPlaybackState);
+	static void OnPostLoadMap(UWorld* World, TSharedRef<const UE::MovieScene::FSharedPlaybackState> SharedPlaybackState, FMovieSceneContext LastContext);
 	static void OnEndPlayMap();
 
 	static FDelegateHandle PreLoadMapHandle;
